CIS fixes.

diff --git a/Engine/Source/Programs/MinidumpDiagnostics/MinidumpDiagnostics.Build.cs b/Engine/Source/Programs/MinidumpDiagnostics/MinidumpDiagnostics.Build.cs
--- a/Engine/Source/Programs/MinidumpDiagnostics/MinidumpDiagnostics.Build.cs
+++ b/Engine/Source/Programs/MinidumpDiagnostics/MinidumpDiagnostics.Build.cs
@@ -13,6 +13,7 @@ public class MinidumpDiagnostics : ModuleRules
 			new string[] 
 			{ 
 				"Core",
+				"ApplicationCore",
 				"CrashDebugHelper",
 				"PerforceSourceControl",
 				"Projects"
